SlateRemote: Documentation cleanup pass

diff --git a/Engine/Plugins/Slate/SlateRemote/Source/SlateRemote/Private/Server/SlateRemoteServer.h b/Engine/Plugins/Slate/SlateRemote/Source/SlateRemote/Private/Server/SlateRemoteServer.h
--- a/Engine/Plugins/Slate/SlateRemote/Source/SlateRemote/Private/Server/SlateRemoteServer.h
+++ b/Engine/Plugins/Slate/SlateRemote/Source/SlateRemote/Private/Server/SlateRemoteServer.h
@@ -33,7 +33,6 @@ public:
 	 * Starts the server.
 	 *
 	 * @param ServerEndpoint The network endpoint to listen on.
-	 *
 	 * @return true if the server was started, false otherwise.
 	 */
 	bool StartServer( const FIPv4Endpoint& ServerEndpoint );
diff --git a/Engine/Plugins/Slate/SlateRemote/Source/SlateRemote/Private/Shared/SlateRemoteSettings.h b/Engine/Plugins/Slate/SlateRemote/Source/SlateRemote/Private/Shared/SlateRemoteSettings.h
--- a/Engine/Plugins/Slate/SlateRemote/Source/SlateRemote/Private/Shared/SlateRemoteSettings.h
+++ b/Engine/Plugins/Slate/SlateRemote/Source/SlateRemote/Private/Shared/SlateRemoteSettings.h
@@ -9,6 +9,9 @@
 #include "SlateRemoteSettings.generated.h"
 
 
+/**
+ * Implements the settings for the Slate Remote plug-in.
+ */
 UCLASS(config=Engine)
 class USlateRemoteSettings
 	: public UObject
@@ -17,21 +20,15 @@ class USlateRemoteSettings
 
 public:
 
-	/**
-	 * Whether the Slate Remote server is enabled.
-	 */
+	/** Whether the Slate Remote server is enabled. */
 	UPROPERTY(config, EditAnywhere, Category=RemoteServer)
 	bool EnableRemoteServer;
 
-	/**
-	 * The IP endpoint to listen to when the Remote Server runs in the Editor.
-	 */
+	/** The IP endpoint to listen to when the Remote Server runs in the Editor. */
 	UPROPERTY(config, EditAnywhere, Category=RemoteServer)
 	FString EditorServerEndpoint;
 
-	/**
-	 * The IP endpoint to listen to when the Remote Server runs in a game.
-	 */
+	/** The IP endpoint to listen to when the Remote Server runs in a game. */
 	UPROPERTY(config, EditAnywhere, Category=RemoteServer)
 	FString GameServerEndpoint;
 };
