Disabled SlateRemote plug-ins in commandlets (UE-10015)

diff --git a/Engine/Plugins/Slate/SlateRemote/SlateRemote.uplugin b/Engine/Plugins/Slate/SlateRemote/SlateRemote.uplugin
--- a/Engine/Plugins/Slate/SlateRemote/SlateRemote.uplugin
+++ b/Engine/Plugins/Slate/SlateRemote/SlateRemote.uplugin
@@ -15,7 +15,7 @@
 	[
 		{
 			"Name" : "SlateRemote",
-			"Type" : "Runtime",
+			"Type" : "RuntimeNoCommandlet",
 			"LoadingPhase" : "Default",
 			"BlacklistPlatforms" : [ "HTML5" ]
 		}
diff --git a/Engine/Plugins/Slate/SlateRemote/Source/SlateRemote/Private/SlateRemoteModule.cpp b/Engine/Plugins/Slate/SlateRemote/Source/SlateRemote/Private/SlateRemoteModule.cpp
--- a/Engine/Plugins/Slate/SlateRemote/Source/SlateRemote/Private/SlateRemoteModule.cpp
+++ b/Engine/Plugins/Slate/SlateRemote/Source/SlateRemote/Private/SlateRemoteModule.cpp
@@ -168,18 +168,10 @@ protected:
 	/**
 	 * Checks whether the Slate Remote server is supported.
 	 *
-	 * @todo gmp: this should be moved into an Engine module, so it can be shared with other transports
-	 *
 	 * @return true if networked transport is supported, false otherwise.
 	 */
 	bool SupportsSlateRemote( ) const
 	{
-		// disallow in Shipping and Test configurations
-		if ((FApp::GetBuildConfiguration() == EBuildConfigurations::Shipping) || (FApp::GetBuildConfiguration() == EBuildConfigurations::Test))
-		{
-			return false;
-		}
-
 		// disallow for commandlets
 		if (IsRunningCommandlet())
 		{
