SlateRemote: code and documentation cleanup pass

diff --git a/Engine/Plugins/Slate/SlateRemote/Source/SlateRemote/Private/Server/SlateRemoteServer.h b/Engine/Plugins/Slate/SlateRemote/Source/SlateRemote/Private/Server/SlateRemoteServer.h
--- a/Engine/Plugins/Slate/SlateRemote/Source/SlateRemote/Private/Server/SlateRemoteServer.h
+++ b/Engine/Plugins/Slate/SlateRemote/Source/SlateRemote/Private/Server/SlateRemoteServer.h
@@ -18,10 +18,8 @@ public:
 	 */
 	FSlateRemoteServer( ISocketSubsystem& InSocketSubsystem, const FIPv4Endpoint& InServerEndpoint );
 
-	/**
-	 * Destructor.
-	 */
-	~FSlateRemoteServer( );
+	/** Destructor. */
+	~FSlateRemoteServer();
 
 public:
 
@@ -33,82 +31,60 @@ public:
 	 */
 	bool StartServer( const FIPv4Endpoint& ServerEndpoint );
 
-	/**
-	 * Stops the server.
-	 */
-	void StopServer( );
+	/** Stops the server. */
+	void StopServer();
 
 protected:
 
-	/**
-	 * Processes a DT_Gyro message.
-	 *
-	 * @param Message The message to process.
-	 */
+	/** Processes a DT_Gyro message. */
 	void ProcessGyroMessage( const FSlateRemoteServerMessage& Message );
 
-	/**
-	 * Processes a DT_Motion message.
-	 *
-	 * @param Message The message to process.
-	 */
+	/** Processes a DT_Motion message. */
 	void ProcessMotionMessage( const FSlateRemoteServerMessage& Message );
 
-	/**
-	 * Processes a DT_Ping message.
-	 *
-	 * @param Message The message to process.
-	 */
+	/** Processes a DT_Ping message. */
 	void ProcessPingMessage( const FSlateRemoteServerMessage& Message );
 
-	/**
-	 * Processes a DT_Tilt message.
-	 *
-	 * @param Message The message to process.
-	 */
+	/** Processes a DT_Tilt message. */
 	void ProcessTiltMessage( const FSlateRemoteServerMessage& Message );
 
-	/**
-	 * Processes a DT_Touch message.
-	 *
-	 * @param Message The message to process.
-	 */
+	/** Processes a DT_Touch message. */
 	void ProcessTouchMessage( const FSlateRemoteServerMessage& Message );
 
 private:
 
-	// Callback for when the ticker fires.
+	/** Callback for when the ticker fires. */
 	bool HandleTicker( float DeltaTime );
 
 private:
 
-	// The widget path to the game viewport.
+	/** The widget path to the game viewport. */
 	FWeakWidgetPath GameViewportWidgetPath;
 
-	// Highest message ID (must handle wrapping around at 16 bits).
+	/** Highest message ID (must handle wrapping around at 16 bits). */
 	uint16 HighestMessageReceived;
 
-	// The socket to send image data on, will be initialized in first tick.
+	/** The socket to send image data on, will be initialized in first tick. */
 	FSocket* ImageSocket;
 
-	// The address of the most recent UDKRemote to talk to us, this is who we reply to.
+	/** The address of the most recent UDKRemote to talk to us, this is who we reply to. */
 	TSharedRef<FInternetAddr> ReplyAddr;
 
-	// The socket to listen on, will be initialized in first tick.
+	/** The socket to listen on, will be initialized in first tick. */
 	FSocket* ServerSocket;
 
-	// Holds a pointer to the socket sub-system being used.
+	/** Holds a pointer to the socket sub-system being used. */
 	ISocketSubsystem& SocketSubsystem;
 
-	// Holds a delegate to be invoked when the server ticks.
+	/** Holds a delegate to be invoked when the server ticks. */
 	FTickerDelegate TickDelegate;
 
-	// Handle to the registered TickDelegate.
+	/** Handle to the registered TickDelegate. */
 	FDelegateHandle TickDelegateHandle;
 
-	// The time since the last message was received from the Slate Remote application.
+	/** The time since the last message was received from the Slate Remote application. */
 	float TimeSinceLastPing;
 
-	// Ever increasing timestamp to send to the input system.
+	/** Ever increasing timestamp to send to the input system. */
 	double Timestamp;
 };
diff --git a/Engine/Plugins/Slate/SlateRemote/Source/SlateRemote/Private/SlateRemoteModule.cpp b/Engine/Plugins/Slate/SlateRemote/Source/SlateRemote/Private/SlateRemoteModule.cpp
--- a/Engine/Plugins/Slate/SlateRemote/Source/SlateRemote/Private/SlateRemoteModule.cpp
+++ b/Engine/Plugins/Slate/SlateRemote/Source/SlateRemote/Private/SlateRemoteModule.cpp
@@ -20,7 +20,7 @@ public:
 
 	// IModuleInterface interface
 
-	virtual void StartupModule( ) override
+	virtual void StartupModule() override
 	{
 		if (!SupportsSlateRemote())
 		{
@@ -51,7 +51,7 @@ public:
 		RestartServices();
 	}
 
-	virtual void ShutdownModule( ) override
+	virtual void ShutdownModule() override
 	{
 		// unregister application events
 		FCoreDelegates::ApplicationHasReactivatedDelegate.RemoveAll(this);
@@ -69,17 +69,15 @@ public:
 		ShutdownRemoteServer();
 	}
 
-	virtual bool SupportsDynamicReloading( ) override
+	virtual bool SupportsDynamicReloading() override
 	{
 		return true;
 	}
 
 protected:
 
-	/**
-	 * Initializes the Slate Remote server with the current settings.
-	 */
-	void InitializeRemoteServer( )
+	/** Initializes the Slate Remote server with the current settings. */
+	void InitializeRemoteServer()
 	{
 		ShutdownRemoteServer();
 
@@ -137,10 +135,8 @@ protected:
 		}
 	}
 
-	/**
-	 * Restarts the services that this modules provides.
-	 */
-	void RestartServices( )
+	/** Restarts the services that this modules provides. */
+	void RestartServices()
 	{
 		const USlateRemoteSettings& Settings = *GetDefault<USlateRemoteSettings>();
 
@@ -157,10 +153,8 @@ protected:
 		}
 	}
 
-	/**
-	 * Shuts down the Slate Remote server.
-	 */
-	void ShutdownRemoteServer( )
+	/** Shuts down the Slate Remote server. */
+	void ShutdownRemoteServer()
 	{
 		RemoteServer.Reset();
 	}
@@ -170,7 +164,7 @@ protected:
 	 *
 	 * @return true if networked transport is supported, false otherwise.
 	 */
-	bool SupportsSlateRemote( ) const
+	bool SupportsSlateRemote() const
 	{
 		// disallow for commandlets
 		if (IsRunningCommandlet())
@@ -183,20 +177,20 @@ protected:
 
 private:
 
-	// Callback for when an has been reactivated (i.e. return from sleep on iOS).
-	void HandleApplicationHasReactivated( )
+	/** Callback for when an has been reactivated (i.e. return from sleep on iOS). */
+	void HandleApplicationHasReactivated()
 	{
 		RestartServices();
 	}
 
-	// Callback for when the application will be deactivated (i.e. sleep on iOS).
-	void HandleApplicationWillDeactivate( )
+	/** Callback for when the application will be deactivated (i.e. sleep on iOS). */
+	void HandleApplicationWillDeactivate()
 	{
 		ShutdownRemoteServer();
 	}
 
-	// Callback for when the settings were saved.
-	bool HandleSettingsSaved( )
+	/** Callback for when the settings were saved. */
+	bool HandleSettingsSaved()
 	{
 		RestartServices();
 
@@ -205,7 +199,7 @@ private:
 
 private:
 
-	// Holds the Slate Remote server.
+	/** Holds the Slate Remote server. */
 	TSharedPtr<FSlateRemoteServer> RemoteServer;
 };
 
diff --git a/Engine/Plugins/Slate/SlateRemote/Source/SlateRemote/Private/SlateRemotePrivatePCH.h b/Engine/Plugins/Slate/SlateRemote/Source/SlateRemote/Private/SlateRemotePrivatePCH.h
--- a/Engine/Plugins/Slate/SlateRemote/Source/SlateRemote/Private/SlateRemotePrivatePCH.h
+++ b/Engine/Plugins/Slate/SlateRemote/Source/SlateRemote/Private/SlateRemotePrivatePCH.h
@@ -17,19 +17,13 @@
 /* Private constants
  *****************************************************************************/
 
-/**
- * Defines the default IP endpoint for the Slate Remote server running in the Editor.
- */
+/** Defines the default IP endpoint for the Slate Remote server running in the Editor. */
 #define SLATE_REMOTE_SERVER_DEFAULT_EDITOR_ENDPOINT FIPv4Endpoint(FIPv4Address(0, 0, 0, 0), 41766)
 
-/**
- * Defines the default IP endpoint for the Slate Remote server running in a game.
- */
+/** Defines the default IP endpoint for the Slate Remote server running in a game. */
 #define SLATE_REMOTE_SERVER_DEFAULT_GAME_ENDPOINT FIPv4Endpoint(FIPv4Address(0, 0, 0, 0), 41765)
 
-/**
- * Defines the protocol version of the UDP message transport.
- */
+/** Defines the protocol version of the UDP message transport. */
 #define SLATE_REMOTE_SERVER_PROTOCOL_VERSION 1
 
 
